Json Utilties / Json Schema.

Trivial precaution -
Clearing visitor stack at start of root visitor method, before then adding root visitor stack element.

#Synced-CL 49094785

diff --git a/Engine/Source/Editor/JsonUtilitiesEditor/Private/JsonSchema/JsonSchemaGeneratorEditor.cpp b/Engine/Source/Editor/JsonUtilitiesEditor/Private/JsonSchema/JsonSchemaGeneratorEditor.cpp
--- a/Engine/Source/Editor/JsonUtilitiesEditor/Private/JsonSchema/JsonSchemaGeneratorEditor.cpp
+++ b/Engine/Source/Editor/JsonUtilitiesEditor/Private/JsonSchema/JsonSchemaGeneratorEditor.cpp
@@ -29,7 +29,9 @@ public:
 	{
 		const EVisitorStackElementFlags Flags = (VisitorTarget.IsType<FValidConstUStruct>() ? 
 			EVisitorStackElementFlags::IsRootStruct : EVisitorStackElementFlags::None);
-    		
+    
+		VisitorStack.Empty(); // Precaution.
+
 		const FVisitorStackElement RootStackElement(VisitorTarget, Flags, MemberPath, PropertyFilter, nullptr); 
 		VisitorStack.Add(RootStackElement);
 		
diff --git a/Engine/Source/Runtime/JsonUtilities/Private/JsonSchema/JsonSchemaGenerator.cpp b/Engine/Source/Runtime/JsonUtilities/Private/JsonSchema/JsonSchemaGenerator.cpp
--- a/Engine/Source/Runtime/JsonUtilities/Private/JsonSchema/JsonSchemaGenerator.cpp
+++ b/Engine/Source/Runtime/JsonUtilities/Private/JsonSchema/JsonSchemaGenerator.cpp
@@ -58,6 +58,8 @@ public:
 		const EVisitorStackElementFlags Flags = (VisitorTarget.IsType<FValidConstUStruct>() ? 
 			EVisitorStackElementFlags::IsRootStruct : EVisitorStackElementFlags::None);
 		
+		VisitorStack.Empty(); // Precaution.
+		
 		const FVisitorStackElement RootStackElement(VisitorTarget, Flags, MemberPath, PropertyFilter, MakeShared<FJsonObject>()); 
 		VisitorStack.Add(RootStackElement);
 		
