Json Utilities / Json Schema -

Added better handling of defaults, including -
- Defaults for UStruct properties.
- Better handling for enum properties.
- More explicit handling for string properties.

#Synced-CL 49627420

diff --git a/Engine/Source/Editor/JsonUtilitiesEditor/Private/JsonSchema/JsonSchemaGeneratorEditor.cpp b/Engine/Source/Editor/JsonUtilitiesEditor/Private/JsonSchema/JsonSchemaGeneratorEditor.cpp
--- a/Engine/Source/Editor/JsonUtilitiesEditor/Private/JsonSchema/JsonSchemaGeneratorEditor.cpp
+++ b/Engine/Source/Editor/JsonUtilitiesEditor/Private/JsonSchema/JsonSchemaGeneratorEditor.cpp
@@ -5,11 +5,272 @@
 #include "UObject/UnrealType.h"
 #include "EdGraphSchema_K2.h"
 #include "ObjectTools.h"
+#include "UObject/TextProperty.h"
+#include "UObject/StrProperty.h"
+#include "JsonUtilitiesLog.h"
 #include "JsonSchema/JsonSchemaVisitorStack.h"
 #include "JsonSchema/JsonSchemaGenerator.h"
+#include "Misc/ScopeExit.h"
 
 using namespace UE::JsonSchema;
 
+//
+// Internal
+//
+
+namespace FJsonSchemaMetadataPropertyDefaults
+{
+	bool SetNamedNumbersJsonObjectDefaultInPropertyMetadata(
+		const FString& PropertyName,
+		const FString& ParameterDefaultString, const TSharedRef<FJsonObject>& PropertyMetadata, 
+		const TArray<FString>& NumberValueNames)
+	{
+		TArray<FString> Parts;
+		if (ParameterDefaultString.ParseIntoArray(Parts, TEXT(",")) <= 0)
+		{
+			UE_LOG(LogJsonUtilities, Error, TEXT(
+				"Property '%s', default string could not be split into components when delimited by \",\".\n"
+				"Default string: \"%s\""), 
+				*PropertyName, 
+				*ParameterDefaultString);
+			
+			return false;
+		}
+		
+		if (Parts.Num() != NumberValueNames.Num())
+		{
+			UE_LOG(LogJsonUtilities, Error, TEXT(
+				"Property '%s', default string contains %d values when delimited by \",\", but is expected to contain %d values.\n"
+				"Default string: \"%s\""), 
+				*PropertyName, 
+				Parts.Num(), NumberValueNames.Num(),
+				*ParameterDefaultString);
+			
+			return false;
+		}
+
+		const TSharedPtr<FJsonObject> JsonObject = MakeShared<FJsonObject>();
+		for (int32 I = 0; I < Parts.Num(); I++)
+		{
+			const FString NumericString = Parts[I].TrimStartAndEnd();
+			
+			double Value = 0.0;
+			if (NumericString.IsEmpty() || !LexTryParseString(Value, *NumericString))
+			{
+				UE_LOG(LogJsonUtilities, Error, TEXT(
+					"Unable to set default value in metadata for Struct Property '%s' when delimited by \",\". "
+					"Reason: Default string contains non-numeric or empty component \"%s\" at component %d.\n"
+					"Default string: \"%s\""),
+					*PropertyName, 
+					*NumericString, 
+					I,
+					*ParameterDefaultString);
+				
+				return false;
+			}
+				
+			JsonObject->SetNumberField(*NumberValueNames[I], Value);
+		}
+		PropertyMetadata->SetObjectField(TEXT("default"), JsonObject);
+		
+		return true;
+	}
+	
+	bool SetFunctionStructPropertyDefaultInPropertyMetadata(
+		const FStructProperty* StructProperty, const FString& ParameterDefault, 
+		const TSharedRef<FJsonObject>& PropertyMetadata)
+	{	
+		bool bSuccess = false;
+		
+		// This is an struct property. This needs to become a JSON object.
+		if (const UScriptStruct* ScriptStruct = StructProperty->Struct)
+		{
+			const int32 StructMemorySize = ScriptStruct->GetStructureSize(); 
+			uint8* StructMemory = static_cast<uint8*>(FMemory_Alloca_Aligned(StructMemorySize, 
+				ScriptStruct->GetMinAlignment()));
+			FMemory::Memzero(StructMemory, StructMemorySize);
+			
+			ScriptStruct->InitializeStruct(StructMemory);
+			
+			ON_SCOPE_EXIT -> void
+			{
+				ScriptStruct->DestroyStruct(StructMemory);
+			};
+			
+			if (StructProperty->ImportText_Direct(*ParameterDefault, StructMemory, nullptr, PPF_None))
+			{
+				if (const TSharedPtr<FJsonObject> JsonObject = MakeShared<FJsonObject>();
+					FJsonObjectConverter::UStructToJsonObject(
+						ScriptStruct, StructMemory, JsonObject.ToSharedRef()))
+				{
+					PropertyMetadata->SetObjectField(TEXT("default"), JsonObject);
+					
+					bSuccess = true;
+				}
+				else
+				{
+					UE_LOG(LogJsonUtilities, Error, TEXT(
+						"Unable to set default value in metadata for Struct Property '%s'. "
+						"Reason: Property's ScriptStruct type '%s' could not be converted to a JSON object.\n"),
+						*StructProperty->GetName(), 
+						*ScriptStruct->GetName());
+				}
+			}
+			else
+			{
+				// ImportText_Direct() failed, usually because the parameter default string is a simple 
+				// list of numbers, which some properties like FVector will yield. For example, 
+				// FVector's default string will be something like "1.00,2.00,3.00", which 
+				// ImportText_Direct() can't handle (see above.) We still want these sorts of defaults 
+				// to become JSON objects with named values, so we have to handle these parameter default
+				// string manually.
+				// NOTE - FTransform does not need to be handled here, it works with ImportText_Direct() above.
+				
+				if (ScriptStruct == TBaseStructure<FVector>::Get() ||
+					ScriptStruct == TBaseStructure<FIntVector>::Get())
+				{
+					static const TArray<FString> NumberValueNames{"x", "y", "z"}; 
+					bSuccess = SetNamedNumbersJsonObjectDefaultInPropertyMetadata(StructProperty->GetName(), 
+						ParameterDefault, PropertyMetadata, NumberValueNames);
+				}
+				else if (ScriptStruct == TBaseStructure<FVector2D>::Get() ||
+					ScriptStruct == TBaseStructure<FIntPoint>::Get())
+				{
+					static const TArray<FString> NumberValueNames{"x", "y"}; 
+					bSuccess = SetNamedNumbersJsonObjectDefaultInPropertyMetadata(StructProperty->GetName(),
+						ParameterDefault, PropertyMetadata, NumberValueNames);
+				}
+				else if (ScriptStruct == TBaseStructure<FVector4>::Get() ||
+					ScriptStruct == TBaseStructure<FQuat>::Get())
+				{
+					static const TArray<FString> NumberValueNames{"x", "y", "z", "w"}; 
+					bSuccess = SetNamedNumbersJsonObjectDefaultInPropertyMetadata(StructProperty->GetName(),
+						ParameterDefault, PropertyMetadata, NumberValueNames);
+				}
+				else if (ScriptStruct == TBaseStructure<FRotator>::Get())
+				{
+					static const TArray<FString> NumberValueNames{"yaw", "pitch", "roll"}; 
+					bSuccess = SetNamedNumbersJsonObjectDefaultInPropertyMetadata(StructProperty->GetName(),
+						ParameterDefault, PropertyMetadata, NumberValueNames);
+				}
+				else if (ScriptStruct == TBaseStructure<FLinearColor>::Get() ||
+					ScriptStruct == TBaseStructure<FColor>::Get())
+				{
+					static const TArray<FString> NumberValueNames{"r", "g", "b", "a"}; 
+					bSuccess = SetNamedNumbersJsonObjectDefaultInPropertyMetadata(StructProperty->GetName(),
+						ParameterDefault, PropertyMetadata, NumberValueNames);
+				}
+				else
+				{
+					UE_LOG(LogJsonUtilities, Error, TEXT(
+						"Unable to set default value in metadata for Struct Property '%s'. "
+						"Reason: Property's ScriptStruct type '%s' is currently unsupported."),
+						*StructProperty->GetName(), 
+						*ScriptStruct->GetName());
+				}
+			}
+		}
+		else
+		{
+			UE_LOG(LogJsonUtilities, Error, TEXT(
+				"Unable to set default value in metadata for Struct Property '%s'. "
+				"Struct Property has no underlying ScriptStruct. "
+				"It may be part of an unloaded module, or reflection data may be incomplete."), 
+				*StructProperty->GetName());
+		}
+		
+		return bSuccess;
+	}	
+
+	void SetFunctionPropertyDefaultInPropertyMetadata(const UFunction* PropertyOwnerFunction, const FProperty* Property, 
+		const TSharedRef<FJsonObject>& PropertyMetadata)
+	{
+		FString ParameterDefault;
+		UEdGraphSchema_K2::FindFunctionParameterDefaultValue(PropertyOwnerFunction, Property, ParameterDefault);
+		if (!ParameterDefault.IsEmpty())
+		{
+			if (const FEnumProperty* EnumProperty = CastField<FEnumProperty>(Property))
+			{
+				if (const UEnum* Enum = EnumProperty->GetEnum())
+				{
+					if (int64 Value = 0;
+						LexTryParseString(Value, *ParameterDefault))
+					{
+						const FString Name = Enum->GetNameStringByValue(Value);
+						PropertyMetadata->SetStringField(TEXT("default"), (Name.IsEmpty() ? ParameterDefault : Name));
+					}
+					else
+					{
+						PropertyMetadata->SetStringField(TEXT("default"), ParameterDefault);
+					}
+				}
+				else
+				{
+					UE_LOG(LogJsonUtilities, Error, TEXT(
+						"Unable to set default value in metadata for Enum Property '%s'. "
+						"Reason: Property has no underlying enum field. "
+						"It may be part of an unloaded module, or reflection data may be incomplete."), 
+						*Property->GetName());
+				}
+			}
+			else if (const FNumericProperty* NumericProperty = CastField<FNumericProperty>(Property))
+			{
+				if (const UEnum* Enum = NumericProperty->GetIntPropertyEnum())
+				{
+					if (int64 Value = 0;
+						LexTryParseString(Value, *ParameterDefault))
+					{
+						const FString Name = Enum->GetNameStringByValue(Value);
+						PropertyMetadata->SetStringField(TEXT("default"), (Name.IsEmpty() ? ParameterDefault : Name));
+					}
+					else
+					{
+						UE_LOG(LogJsonUtilities, Error, TEXT(
+							"Unable to set default value in metadata for Numeric Enum Property '%s'. "
+							"Reason: The default string for the Property is not numeric, and cannot be resolved into an Enum entry.\n"
+							"Default string: \"%s\""), 
+							*Property->GetName(),
+							*ParameterDefault);
+					}
+				}
+				else
+				{
+					// This is a normal numeric property. Convert default string to double (ints will end up being rounded in the JSON.)
+					PropertyMetadata->SetNumberField(TEXT("default"), FCString::Atod(*ParameterDefault));
+				}
+			}
+			else if (Property->IsA<FBoolProperty>())
+			{
+				// Default string will be either 'true' or 'false' only, and needs to become a bool.
+				PropertyMetadata->SetBoolField(TEXT("default"), (ParameterDefault.Compare(TEXT("true"), ESearchCase::IgnoreCase) == 0));
+			}
+			else if (Property->IsA<FStrProperty>() || Property->IsA<FNameProperty>() || Property->IsA<FTextProperty>())
+			{
+				// Default is one of the string properties, and we can use the default as-is.
+				PropertyMetadata->SetStringField(TEXT("default"), ParameterDefault);
+			}
+			else if (const FStructProperty* StructProperty = CastField<FStructProperty>(Property))
+			{
+				// This is an struct property. This needs to become a JSON object, which 
+				// requires some special handling.
+				SetFunctionStructPropertyDefaultInPropertyMetadata(StructProperty, ParameterDefault, PropertyMetadata);
+			}
+			else if (Property->IsA<FObjectProperty>())
+			{
+				// Do not attempt to set defaults for this property type.
+			}
+			else
+			{
+				UE_LOG(LogJsonUtilities, Error, TEXT(
+					"Unable to set default value in metadata for Property '%s'. "
+					"Reason: Property type '%s' is not recognized."), 
+					*Property->GetName(),
+					*Property->GetClass()->GetName());
+			}
+		}
+	}
+};
+
 //
 // Internal FJsonSchemaMetadataVisitor
 //
@@ -103,26 +364,7 @@ public:
 		// If function, add any parameter defaults.
 		if (PropertyOwnerFunction)
 		{
-			FString ParameterDefault;
-			UEdGraphSchema_K2::FindFunctionParameterDefaultValue(PropertyOwnerFunction, Property, ParameterDefault);
-			if (!ParameterDefault.IsEmpty())
-			{
-				if (Property->IsA<FNumericProperty>())
-				{
-					// Convert default string to double, ints will end up being rounded in the JSON.
-					PropertyMetadata->SetNumberField(TEXT("default"), FCString::Atod(*ParameterDefault));
-				}
-				else if (Property->IsA<FBoolProperty>())
-				{
-					// Default string will be either 'true' or 'false' only.
-					PropertyMetadata->SetBoolField(TEXT("default"), ParameterDefault.Equals("true"));
-				}
-				else
-				{
-					// Default string works for both strings and enums.
-					PropertyMetadata->SetStringField(TEXT("default"), ParameterDefault);
-				}
-			}
+			FJsonSchemaMetadataPropertyDefaults::SetFunctionPropertyDefaultInPropertyMetadata(PropertyOwnerFunction, Property, PropertyMetadata);
 		}
 
 		if (PropertyOwnerFunction && EnumHasAnyFlags(Flags, EVisitorStackElementFlags::IsFunctionWorldContextProperty))
diff --git a/Engine/Source/Runtime/JsonUtilities/Private/JsonUtilitiesLog.cpp b/Engine/Source/Runtime/JsonUtilities/Private/JsonUtilitiesLog.cpp
new file mode 100644
--- /dev/null
+++ b/Engine/Source/Runtime/JsonUtilities/Private/JsonUtilitiesLog.cpp
@@ -0,0 +1,5 @@
+// Copyright Epic Games, Inc. All Rights Reserved.
+
+#include "JsonUtilitiesLog.h"
+
+DEFINE_LOG_CATEGORY(LogJsonUtilities);
diff --git a/Engine/Source/Runtime/JsonUtilities/Private/JsonUtilitiesLog.h b/Engine/Source/Runtime/JsonUtilities/Private/JsonUtilitiesLog.h
new file mode 100644
--- /dev/null
+++ b/Engine/Source/Runtime/JsonUtilities/Private/JsonUtilitiesLog.h
@@ -0,0 +1,7 @@
+// Copyright Epic Games, Inc. All Rights Reserved.
+
+#pragma once
+
+#include "Logging/LogMacros.h"
+
+JSONUTILITIES_API DECLARE_LOG_CATEGORY_EXTERN(LogJsonUtilities, Log, Log);
