Json Utilities / Json Schema

For function param defaults -
- Special handling of soft object path default.
- Special handling of soft class path default.
- Errors out when param type cannot be supported in Json (i.e. UObject*, etc.)
- Removed WorldContext support.
- Additional error conditions detected.
- Removed recently added JsonUtilitiesLog.

#Synced-CL 49747421

diff --git a/Engine/Source/Editor/JsonUtilitiesEditor/Private/JsonSchema/JsonSchemaGeneratorEditor.cpp b/Engine/Source/Editor/JsonUtilitiesEditor/Private/JsonSchema/JsonSchemaGeneratorEditor.cpp
--- a/Engine/Source/Editor/JsonUtilitiesEditor/Private/JsonSchema/JsonSchemaGeneratorEditor.cpp
+++ b/Engine/Source/Editor/JsonUtilitiesEditor/Private/JsonSchema/JsonSchemaGeneratorEditor.cpp
@@ -2,49 +2,71 @@
 
 #include "JsonSchema/JsonSchemaGeneratorEditor.h"
 
-#include "UObject/UnrealType.h"
 #include "EdGraphSchema_K2.h"
 #include "ObjectTools.h"
+#include "UObject/UnrealType.h"
 #include "UObject/TextProperty.h"
+#include "UObject/EnumProperty.h"
 #include "UObject/StrProperty.h"
-#include "JsonUtilitiesLog.h"
+#include "UObject/Utf8StrProperty.h"
+#include "UObject/PropertyOptional.h"
+#include "Misc/ScopeExit.h"
 #include "JsonSchema/JsonSchemaVisitorStack.h"
 #include "JsonSchema/JsonSchemaGenerator.h"
-#include "Misc/ScopeExit.h"
 
 using namespace UE::JsonSchema;
 
 //
-// Internal
+// Internal FJsonSchemaMetadataPropertyDefaults
+//
+// Used for setting function struct param property defaults in the Json schema metadata.
 //
 
 namespace FJsonSchemaMetadataPropertyDefaults
 {
-	bool SetNamedNumbersJsonObjectDefaultInPropertyMetadata(
-		const FString& PropertyName,
-		const FString& ParameterDefaultString, const TSharedRef<FJsonObject>& PropertyMetadata, 
-		const TArray<FString>& NumberValueNames)
+	static const FString ParameterDefaultStringErrorPhraseString = 
+		"Unable to convert default string to default value for metadata"; 
+	
+	template <typename T>
+	void GetSoftPathStringsFromStructMemory(const uint8* const StructMemory, 
+		FString& AssetPathDebugString, 
+		FString& PackageNameString,
+		FString& AssetNameString,
+		FString& SubPathString,
+		bool& bIsNullSoftPath)
 	{
-		TArray<FString> Parts;
-		if (ParameterDefaultString.ParseIntoArray(Parts, TEXT(",")) <= 0)
-		{
-			UE_LOG(LogJsonUtilities, Error, TEXT(
-				"Property '%s', default string could not be split into components when delimited by \",\".\n"
-				"Default string: \"%s\""), 
-				*PropertyName, 
-				*ParameterDefaultString);
-			
-			return false;
-		}
+		static_assert(std::is_same_v<T, FSoftObjectPath> || std::is_same_v<T, FSoftClassPath>,
+			"GetSoftPathStringsFromStructMemory() - T must be FSoftObjectPath or FSoftClassPath.");
 		
+		const FSoftObjectPath* SoftObjectPath = reinterpret_cast<const T*>(StructMemory);
+		AssetPathDebugString = SoftObjectPath->GetAssetPathString();
+		PackageNameString = SoftObjectPath->GetLongPackageName();
+		AssetNameString = SoftObjectPath->GetAssetName();
+		SubPathString = SoftObjectPath->GetSubPathString();
+		bIsNullSoftPath = SoftObjectPath->IsNull();
+	}
+	
+	bool SetFunctionStructPropertyDefaultAsNamedNumbersInPropertyMetadata(
+		const FString& FunctionDebugName, const FString& FunctionParamPropertyDebugName,
+		const FString& ParameterDefaultString, 
+		const TArray<FString>& NumberValueNames,
+		const TSharedRef<FJsonObject>& OutPropertyMetadata)
+	{
+		static const FString Delimiter = TEXT(",");
+		static const FString DelimiterErrorPhraseTemplate = TEXT("(when delimited by \"") + Delimiter + TEXT("\")");
+
+		TArray<FString> Parts;
+		ParameterDefaultString.ParseIntoArray(Parts, *Delimiter);
+
 		if (Parts.Num() != NumberValueNames.Num())
 		{
-			UE_LOG(LogJsonUtilities, Error, TEXT(
-				"Property '%s', default string contains %d values when delimited by \",\", but is expected to contain %d values.\n"
-				"Default string: \"%s\""), 
-				*PropertyName, 
-				Parts.Num(), NumberValueNames.Num(),
-				*ParameterDefaultString);
+			UE_LOG(LogJson, Error, TEXT(
+				"Function \"%s\", param property \"%s\". " 
+				"%s.\n"
+				"Reason: Default string contains %d values, but %d values were expected."),
+				*FunctionDebugName, *FunctionParamPropertyDebugName,
+				*ParameterDefaultStringErrorPhraseString,
+				Parts.Num(), NumberValueNames.Num());
 			
 			return false;
 		}
@@ -54,34 +76,36 @@ namespace FJsonSchemaMetadataPropertyDefaults
 		{
 			const FString NumericString = Parts[I].TrimStartAndEnd();
 			
-			double Value = 0.0;
-			if (NumericString.IsEmpty() || !LexTryParseString(Value, *NumericString))
+			if (double Value = 0.0;
+				LexTryParseString(Value, *NumericString))
 			{
-				UE_LOG(LogJsonUtilities, Error, TEXT(
-					"Unable to set default value in metadata for Struct Property '%s' when delimited by \",\". "
-					"Reason: Default string contains non-numeric or empty component \"%s\" at component %d.\n"
+				JsonObject->SetNumberField(*NumberValueNames[I], Value);
+			}				
+			else
+			{
+				UE_LOG(LogJson, Error, TEXT(
+					"Function \"%s\", param property \"%s\". " 
+					"%s.\n"
+					"Reason: Parameter default string contains non-double numeric string \"%s\".\n"
 					"Default string: \"%s\""),
-					*PropertyName, 
+					*FunctionDebugName, *FunctionParamPropertyDebugName,
+					*ParameterDefaultStringErrorPhraseString,
 					*NumericString, 
-					I,
 					*ParameterDefaultString);
-				
+	
 				return false;
 			}
-				
-			JsonObject->SetNumberField(*NumberValueNames[I], Value);
 		}
-		PropertyMetadata->SetObjectField(TEXT("default"), JsonObject);
+		OutPropertyMetadata->SetObjectField(TEXT("default"), JsonObject);
 		
 		return true;
 	}
 	
 	bool SetFunctionStructPropertyDefaultInPropertyMetadata(
-		const FStructProperty* StructProperty, const FString& ParameterDefault, 
-		const TSharedRef<FJsonObject>& PropertyMetadata)
+		const FString& FunctionDebugName,
+		const FStructProperty* StructProperty, const FString& ParameterDefaultString, 
+		const TSharedRef<FJsonObject>& OutPropertyMetadata)
 	{	
-		bool bSuccess = false;
-		
 		// This is an struct property. This needs to become a JSON object.
 		if (const UScriptStruct* ScriptStruct = StructProperty->Struct)
 		{
@@ -97,23 +121,91 @@ namespace FJsonSchemaMetadataPropertyDefaults
 				ScriptStruct->DestroyStruct(StructMemory);
 			};
 			
-			if (StructProperty->ImportText_Direct(*ParameterDefault, StructMemory, nullptr, PPF_None))
+			if (StructProperty->ImportText_Direct(*ParameterDefaultString, StructMemory, 
+				nullptr, PPF_None))
 			{
-				if (const TSharedPtr<FJsonObject> JsonObject = MakeShared<FJsonObject>();
-					FJsonObjectConverter::UStructToJsonObject(
-						ScriptStruct, StructMemory, JsonObject.ToSharedRef()))
+				// 'Import text' process succeeded, and we can convert the struct to Json directly.
+				// But we also need special handling for some particular struct types.
+				
+				if (ScriptStruct == TBaseStructure<FSoftObjectPath>::Get() ||
+					ScriptStruct == TBaseStructure<FSoftClassPath>::Get())
 				{
-					PropertyMetadata->SetObjectField(TEXT("default"), JsonObject);
+					// Special handling of FSoftObjectPath/FSoftClassPath, to prevent getting the 
+					// specialized default, which is a single combined string.
+					// NOTE - We initialize from raw memory, so we can't safely use type
+					// punning between FSoftObjectPath and derived FSoftClassPath, due to 
+					// potential differences in memory layout.
 					
-					bSuccess = true;
+					FString AssetPathDebugString, PackageNameString, AssetNameString, SubPathString;
+					bool bIsNullSoftPath = false;
+					if (ScriptStruct == TBaseStructure<FSoftObjectPath>::Get())
+					{
+						GetSoftPathStringsFromStructMemory<FSoftObjectPath>(StructMemory, 
+							AssetPathDebugString, PackageNameString, AssetNameString, SubPathString, 
+							bIsNullSoftPath);
+					}
+					else
+					{
+						GetSoftPathStringsFromStructMemory<FSoftClassPath>(StructMemory, 
+							AssetPathDebugString, PackageNameString, AssetNameString, SubPathString, 
+							bIsNullSoftPath);
+					}
+					
+					if (bIsNullSoftPath || (!PackageNameString.IsEmpty() && !AssetNameString.IsEmpty()))
+					{
+						const TSharedPtr<FJsonObject> PackageAssetJsonObject = MakeShared<FJsonObject>();
+						PackageAssetJsonObject->SetStringField(TEXT("packageName"), PackageNameString);
+						PackageAssetJsonObject->SetStringField(TEXT("assetName"), AssetNameString);
+						
+						const TSharedPtr<FJsonObject> JsonObject = MakeShared<FJsonObject>();
+						JsonObject->SetObjectField(TEXT("assetPath"), PackageAssetJsonObject);
+						JsonObject->SetStringField(TEXT("subPathString"), SubPathString);
+						OutPropertyMetadata->SetObjectField(TEXT("default"), JsonObject);
+
+						return true;
+					}
+					else
+					{
+						UE_LOG(LogJson, Error, TEXT(
+							"Function \"%s\", param property \"%s\". "
+							"%s.\n"
+							"Reason: Property is non-empty soft object or class path struct property, "
+							"but asset path \"%s\" is missing package name \"%s\" or asset name \"%s\".\n"),
+							*FunctionDebugName, 
+							*StructProperty->GetName(), 
+							*ParameterDefaultStringErrorPhraseString,
+							*AssetPathDebugString, 
+							*PackageNameString, 
+							*AssetNameString);
+
+						return false;
+					}
 				}
 				else
 				{
-					UE_LOG(LogJsonUtilities, Error, TEXT(
-						"Unable to set default value in metadata for Struct Property '%s'. "
-						"Reason: Property's ScriptStruct type '%s' could not be converted to a JSON object.\n"),
-						*StructProperty->GetName(), 
-						*ScriptStruct->GetName());
+					// All other struct types can use standard Struct->JsonObject conversion.
+					
+					if (const TSharedPtr<FJsonObject> JsonObject = MakeShared<FJsonObject>();
+						FJsonObjectConverter::UStructToJsonObject(
+							ScriptStruct, StructMemory, JsonObject.ToSharedRef()))
+					{
+						OutPropertyMetadata->SetObjectField(TEXT("default"), JsonObject);
+
+						return true;
+					}
+					else
+					{
+						UE_LOG(LogJson, Error, TEXT(
+							"Function \"%s\", param property \"%s\". "
+							"%s.\n"
+							"Reason: Property is a struct property, "
+							"but property's ScriptStruct type \"%s\" could not be converted to a JSON object.\n"),
+							*FunctionDebugName, 
+							*StructProperty->GetName(), 
+							*ParameterDefaultStringErrorPhraseString,
+							*ScriptStruct->GetPathName());
+						return false;
+					}
 				}
 			}
 			else
@@ -126,91 +218,108 @@ namespace FJsonSchemaMetadataPropertyDefaults
 				// string manually.
 				// NOTE - FTransform does not need to be handled here, it works with ImportText_Direct() above.
 				
+				TArray<FString> NumberValueNames;
+				
 				if (ScriptStruct == TBaseStructure<FVector>::Get() ||
 					ScriptStruct == TBaseStructure<FIntVector>::Get())
 				{
-					static const TArray<FString> NumberValueNames{"x", "y", "z"}; 
-					bSuccess = SetNamedNumbersJsonObjectDefaultInPropertyMetadata(StructProperty->GetName(), 
-						ParameterDefault, PropertyMetadata, NumberValueNames);
+					NumberValueNames = {"x", "y", "z"}; 
 				}
 				else if (ScriptStruct == TBaseStructure<FVector2D>::Get() ||
 					ScriptStruct == TBaseStructure<FIntPoint>::Get())
 				{
-					static const TArray<FString> NumberValueNames{"x", "y"}; 
-					bSuccess = SetNamedNumbersJsonObjectDefaultInPropertyMetadata(StructProperty->GetName(),
-						ParameterDefault, PropertyMetadata, NumberValueNames);
+					NumberValueNames = {"x", "y"}; 
 				}
 				else if (ScriptStruct == TBaseStructure<FVector4>::Get() ||
 					ScriptStruct == TBaseStructure<FQuat>::Get())
 				{
-					static const TArray<FString> NumberValueNames{"x", "y", "z", "w"}; 
-					bSuccess = SetNamedNumbersJsonObjectDefaultInPropertyMetadata(StructProperty->GetName(),
-						ParameterDefault, PropertyMetadata, NumberValueNames);
+					NumberValueNames = {"x", "y", "z", "w"}; 
 				}
 				else if (ScriptStruct == TBaseStructure<FRotator>::Get())
 				{
-					static const TArray<FString> NumberValueNames{"yaw", "pitch", "roll"}; 
-					bSuccess = SetNamedNumbersJsonObjectDefaultInPropertyMetadata(StructProperty->GetName(),
-						ParameterDefault, PropertyMetadata, NumberValueNames);
+					NumberValueNames = {"yaw", "pitch", "roll"}; 
 				}
 				else if (ScriptStruct == TBaseStructure<FLinearColor>::Get() ||
 					ScriptStruct == TBaseStructure<FColor>::Get())
 				{
-					static const TArray<FString> NumberValueNames{"r", "g", "b", "a"}; 
-					bSuccess = SetNamedNumbersJsonObjectDefaultInPropertyMetadata(StructProperty->GetName(),
-						ParameterDefault, PropertyMetadata, NumberValueNames);
+					NumberValueNames = {"r", "g", "b", "a"}; 
 				}
 				else
 				{
-					UE_LOG(LogJsonUtilities, Error, TEXT(
-						"Unable to set default value in metadata for Struct Property '%s'. "
-						"Reason: Property's ScriptStruct type '%s' is currently unsupported."),
+					UE_LOG(LogJson, Error, TEXT(
+						"Function \"%s\", param property \"%s\". "
+						"%s.\n"
+						"Reason: Property is a struct property, "
+						"but property's ScriptStruct type \"%s\" is currently unsupported."),
+						*FunctionDebugName, 
 						*StructProperty->GetName(), 
-						*ScriptStruct->GetName());
+						*ParameterDefaultStringErrorPhraseString, 
+						*ScriptStruct->GetPathName());
+
+					return false;
 				}
+				
+				return SetFunctionStructPropertyDefaultAsNamedNumbersInPropertyMetadata(
+					FunctionDebugName, StructProperty->GetName(),
+					ParameterDefaultString, NumberValueNames, OutPropertyMetadata);
 			}
 		}
 		else
 		{
-			UE_LOG(LogJsonUtilities, Error, TEXT(
-				"Unable to set default value in metadata for Struct Property '%s'. "
-				"Struct Property has no underlying ScriptStruct. "
-				"It may be part of an unloaded module, or reflection data may be incomplete."), 
-				*StructProperty->GetName());
+			UE_LOG(LogJson, Error, TEXT(
+				"Function \"%s\", param property \"%s\". "
+				"%s.\n"
+				"Reason: Property is struct property, but has no underlying ScriptStruct. "
+				"It may be part of an unloaded module, or reflection data may be incomplete."),
+				*FunctionDebugName, 
+				*StructProperty->GetName(), 
+				*ParameterDefaultStringErrorPhraseString);
+			
+			return false;
 		}
-		
-		return bSuccess;
 	}	
 
-	void SetFunctionPropertyDefaultInPropertyMetadata(const UFunction* PropertyOwnerFunction, const FProperty* Property, 
-		const TSharedRef<FJsonObject>& PropertyMetadata)
+	void SetFunctionPropertyDefaultInPropertyMetadata(
+		const UFunction* PropertyOwnerFunction, const FProperty* Property, 
+		const TSharedRef<FJsonObject>& OutPropertyMetadata)
 	{
-		FString ParameterDefault;
-		UEdGraphSchema_K2::FindFunctionParameterDefaultValue(PropertyOwnerFunction, Property, ParameterDefault);
-		if (!ParameterDefault.IsEmpty())
+		if (Property->HasAnyPropertyFlags(CPF_ReturnParm))
+		{
+			// Return values don't have defaults.
+			return;
+		}
+		
+		const FString FunctionDebugName = PropertyOwnerFunction->GetName();
+		
+		FString ParameterDefaultString;
+		UEdGraphSchema_K2::FindFunctionParameterDefaultValue(PropertyOwnerFunction, Property, ParameterDefaultString);
+		if (!ParameterDefaultString.IsEmpty())
 		{
 			if (const FEnumProperty* EnumProperty = CastField<FEnumProperty>(Property))
 			{
 				if (const UEnum* Enum = EnumProperty->GetEnum())
 				{
 					if (int64 Value = 0;
-						LexTryParseString(Value, *ParameterDefault))
+						LexTryParseString(Value, *ParameterDefaultString))
 					{
 						const FString Name = Enum->GetNameStringByValue(Value);
-						PropertyMetadata->SetStringField(TEXT("default"), (Name.IsEmpty() ? ParameterDefault : Name));
+						OutPropertyMetadata->SetStringField(TEXT("default"), (Name.IsEmpty() ? ParameterDefaultString : Name));
 					}
 					else
 					{
-						PropertyMetadata->SetStringField(TEXT("default"), ParameterDefault);
+						OutPropertyMetadata->SetStringField(TEXT("default"), ParameterDefaultString);
 					}
 				}
 				else
 				{
-					UE_LOG(LogJsonUtilities, Error, TEXT(
-						"Unable to set default value in metadata for Enum Property '%s'. "
-						"Reason: Property has no underlying enum field. "
-						"It may be part of an unloaded module, or reflection data may be incomplete."), 
-						*Property->GetName());
+					UE_LOG(LogJson, Error, TEXT(
+						"Function \"%s\", param property \"%s\". "
+						"%s.\n"
+						"Reason: Property is enum property, but has no underlying enum field. "
+						"It may be part of an unloaded module, or reflection data may be incomplete."),
+						*FunctionDebugName, 
+						*Property->GetName(), 
+						*ParameterDefaultStringErrorPhraseString);
 				}
 			}
 			else if (const FNumericProperty* NumericProperty = CastField<FNumericProperty>(Property))
@@ -218,42 +327,55 @@ namespace FJsonSchemaMetadataPropertyDefaults
 				if (const UEnum* Enum = NumericProperty->GetIntPropertyEnum())
 				{
 					if (int64 Value = 0;
-						LexTryParseString(Value, *ParameterDefault))
+						LexTryParseString(Value, *ParameterDefaultString))
 					{
 						const FString Name = Enum->GetNameStringByValue(Value);
-						PropertyMetadata->SetStringField(TEXT("default"), (Name.IsEmpty() ? ParameterDefault : Name));
+						OutPropertyMetadata->SetStringField(TEXT("default"), (Name.IsEmpty() ? ParameterDefaultString : Name));
 					}
 					else
 					{
-						UE_LOG(LogJsonUtilities, Error, TEXT(
-							"Unable to set default value in metadata for Numeric Enum Property '%s'. "
-							"Reason: The default string for the Property is not numeric, and cannot be resolved into an Enum entry.\n"
-							"Default string: \"%s\""), 
-							*Property->GetName(),
-							*ParameterDefault);
+						UE_LOG(LogJson, Error, TEXT(
+							"Function \"%s\", param property \"%s\". "
+							"%s.\n"
+							"Reason: Property is numeric enum property, "
+							"but the default string for the Property is not integer type "
+							"and cannot be resolved into an Enum entry.\n"
+							"Default string: \"%s\""),
+							*FunctionDebugName, 
+							*Property->GetName(), 
+							*ParameterDefaultStringErrorPhraseString,
+							*ParameterDefaultString);
 					}
 				}
 				else
 				{
 					// This is a normal numeric property. Convert default string to double (ints will end up being rounded in the JSON.)
-					PropertyMetadata->SetNumberField(TEXT("default"), FCString::Atod(*ParameterDefault));
+					OutPropertyMetadata->SetNumberField(TEXT("default"), FCString::Atod(*ParameterDefaultString));
 				}
 			}
 			else if (Property->IsA<FBoolProperty>())
 			{
 				// Default string will be either 'true' or 'false' only, and needs to become a bool.
-				PropertyMetadata->SetBoolField(TEXT("default"), (ParameterDefault.Compare(TEXT("true"), ESearchCase::IgnoreCase) == 0));
+				OutPropertyMetadata->SetBoolField(TEXT("default"), (ParameterDefaultString.Compare(TEXT("true"), ESearchCase::IgnoreCase) == 0));
 			}
 			else if (Property->IsA<FStrProperty>() || Property->IsA<FNameProperty>() || Property->IsA<FTextProperty>())
 			{
 				// Default is one of the string properties, and we can use the default as-is.
-				PropertyMetadata->SetStringField(TEXT("default"), ParameterDefault);
+				OutPropertyMetadata->SetStringField(TEXT("default"), ParameterDefaultString);
 			}
 			else if (const FStructProperty* StructProperty = CastField<FStructProperty>(Property))
 			{
 				// This is an struct property. This needs to become a JSON object, which 
 				// requires some special handling.
-				SetFunctionStructPropertyDefaultInPropertyMetadata(StructProperty, ParameterDefault, PropertyMetadata);
+				if (!SetFunctionStructPropertyDefaultInPropertyMetadata(FunctionDebugName, StructProperty, ParameterDefaultString, OutPropertyMetadata))
+				{
+					UE_LOG(LogJson, Error, TEXT(
+						"Function \"%s\", param property \"%s\". " 
+						"%s.\n"),
+						*FunctionDebugName, 
+						*Property->GetName(), 
+						*ParameterDefaultStringErrorPhraseString);
+				}
 			}
 			else if (Property->IsA<FObjectProperty>())
 			{
@@ -261,10 +383,13 @@ namespace FJsonSchemaMetadataPropertyDefaults
 			}
 			else
 			{
-				UE_LOG(LogJsonUtilities, Error, TEXT(
-					"Unable to set default value in metadata for Property '%s'. "
-					"Reason: Property type '%s' is not recognized."), 
-					*Property->GetName(),
+				UE_LOG(LogJson, Error, TEXT(
+					"Function \"%s\", param property \"%s\". " 
+					"%s.\n"
+					"Reason: Property type \"%s\" is not recognized."),
+					*FunctionDebugName, 
+					*Property->GetName(), 
+					*ParameterDefaultStringErrorPhraseString,
 					*Property->GetClass()->GetName());
 			}
 		}
@@ -320,8 +445,9 @@ public:
 	{
 		OutputEditorMetadata.CurrentPropertyMemberPath = MemberPath;
 		
-		const UFunction* PropertyOwnerFunction = Cast<UFunction>(Property->GetOwnerStruct());
-			
+		const UStruct* PropertyOwnerStruct = Property->GetOwnerStruct();
+		const UFunction* PropertyOwnerFunction = Cast<UFunction>(PropertyOwnerStruct);
+		
 		// Check to see if we should ignore this property.
 		if (PropertyFilter.IsPropertyIgnored(Property))
 		{
@@ -366,11 +492,6 @@ public:
 		{
 			FJsonSchemaMetadataPropertyDefaults::SetFunctionPropertyDefaultInPropertyMetadata(PropertyOwnerFunction, Property, PropertyMetadata);
 		}
-
-		if (PropertyOwnerFunction && EnumHasAnyFlags(Flags, EVisitorStackElementFlags::IsFunctionWorldContextProperty))
-		{
-			PropertyMetadata->SetBoolField(TEXT("isWorldContext"), true);
-		}
 		
 		if (PropertyMetadata->Values.Num())
 		{
@@ -398,11 +519,6 @@ public:
 		
 		const UFunction* StructAsFunction = Cast<UFunction>(Struct);
 		
-		const FString FunctionWorldContext =
-			(StructAsFunction && StructAsFunction->HasMetaData(FBlueprintMetadata::MD_WorldContext)) ?
-			StructAsFunction->GetMetaData(FBlueprintMetadata::MD_WorldContext) :
-			"";
-		
 		if (EnumHasAnyFlags(Flags, EVisitorStackElementFlags::IsRootStruct))
 		{
 			OutputEditorMetadata.RootStructMetadata = MakeShared<FJsonObject>();
@@ -415,11 +531,6 @@ public:
 			{
 				OutputEditorMetadata.RootStructMetadata->SetStringField(TEXT("description"), Description);
 			}
-			
-			if (!FunctionWorldContext.IsEmpty())
-			{
-				OutputEditorMetadata.RootStructMetadata->SetStringField(TEXT("worldContext"), FunctionWorldContext);
-			}
 		}
 		
 		for (const FProperty* Property : TFieldRange<FProperty>(Struct))
@@ -429,13 +540,9 @@ public:
 			UE_JSON_SCHEMA_SCOPED_MEMBER_PATH_PUSH(OutputEditorMetadata.CurrentPropertyMemberPath, PropertyName, 
 				!EnumHasAnyFlags(Flags, EVisitorStackElementFlags::IsRootStruct));
 			
-			const EVisitorStackElementFlags NewFlags = (FunctionWorldContext == PropertyName ?
-				EVisitorStackElementFlags::IsFunctionWorldContextProperty : 
-				EVisitorStackElementFlags::None);
-			
 			VisitorStack.Emplace(
 				FVisitorTarget(TInPlaceType<FValidConstFProperty>(), Property), 
-				NewFlags, OutputEditorMetadata.CurrentPropertyMemberPath, PropertyFilter, nullptr);
+				EVisitorStackElementFlags::None, OutputEditorMetadata.CurrentPropertyMemberPath, PropertyFilter, nullptr);
 		}
 	}
 	
diff --git a/Engine/Source/Runtime/JsonUtilities/Private/JsonSchema/JsonSchemaGenerator.cpp b/Engine/Source/Runtime/JsonUtilities/Private/JsonSchema/JsonSchemaGenerator.cpp
--- a/Engine/Source/Runtime/JsonUtilities/Private/JsonSchema/JsonSchemaGenerator.cpp
+++ b/Engine/Source/Runtime/JsonUtilities/Private/JsonSchema/JsonSchemaGenerator.cpp
@@ -4,12 +4,16 @@
 
 #include "JsonDomBuilder.h"
 #include "JsonObjectConverter.h"
-#include "UObject/TextProperty.h"
-#include "UObject/UnrealType.h"
 #include "Policies/CondensedJsonPrintPolicy.h"
-#include "JsonSchema/JsonSchemaVisitorStack.h"
-
+#include "UObject/UnrealType.h"
+#include "UObject/TextProperty.h"
+#include "UObject/EnumProperty.h"
+#include "UObject/StrProperty.h"
+#include "UObject/Utf8StrProperty.h"
+#include "UObject/PropertyOptional.h"
 #include "JsonUtilitiesModularFeature.h"
+#include "UObject/FieldPathProperty.h"
+#include "JsonSchema/JsonSchemaVisitorStack.h"
 
 using namespace UE::JsonSchema;
 
@@ -37,6 +41,47 @@ namespace
 		
 		return (bSuccess ? OutString : TEXT(""));
 	}
+	
+	bool CanPropertyBeDirectlyRepresentedInJson(const FValidConstFProperty& Property)
+	{
+		return !(
+			// Object, Classes, interfaces (pointer-like types)..
+			Property->IsA<FObjectPropertyBase>() || 
+			Property->IsA<FClassProperty>() ||
+			Property->IsA<FSoftClassProperty>() ||
+			Property->IsA<FInterfaceProperty>() ||	
+			Property->IsA<FFieldPathProperty>() ||
+			// Delegates..
+			Property->IsA<FDelegateProperty>() || 
+			Property->IsA<FMulticastDelegateProperty>() || 
+			Property->IsA<FMulticastInlineDelegateProperty>() || 
+			Property->IsA<FMulticastSparseDelegateProperty>());
+	}
+	
+	bool CanStructBeDirectlyRepresentedInJson(const FValidConstUStruct& Struct, FJsonSchemaEditorMetadata& EditorMetadata)
+	{
+		bool bCanAllPropertiesBeAcceptedForJson = true;
+		for (const FProperty* Property : TFieldRange<FProperty>(Struct))
+		{
+			if (!CanPropertyBeDirectlyRepresentedInJson(Property))
+			{
+				const FString InputOrOutputParamString = (Property->HasAnyPropertyFlags(CPF_ReturnParm) ? 
+					TEXT("return") : TEXT("input"));
+				
+				UE_LOG(LogJson, Error, TEXT(
+					"Illegal property type for Json - " // Warning, this string necessary during testing, do not change.
+					"UFunction '%s' %s param '%s' is type %s, and cannot be directly represented in the JSON schema."),
+					*Struct->GetPathName(), 
+					*InputOrOutputParamString, 
+					*Property->GetName(), 
+					*Property->GetCPPType());
+				
+				bCanAllPropertiesBeAcceptedForJson = false;
+			}
+		}
+		
+		return bCanAllPropertiesBeAcceptedForJson;
+	}
 }
 
 //
@@ -69,21 +114,27 @@ public:
 			
 			if (StackElement.VisitorTarget.IsType<FValidConstUStruct>())
 			{
-				VisitUStruct(StackElement.VisitorTarget.Get<FValidConstUStruct>(), 
+				if (!VisitUStruct(StackElement.VisitorTarget.Get<FValidConstUStruct>(), 
 					StackElement.OutputSchema.ToSharedRef(), StackElement.MemberPath, StackElement.PropertyFilter, 
-					StackElement.Flags);
+					StackElement.Flags))
+				{
+					return nullptr;
+				}
 			}
 			else if (StackElement.VisitorTarget.IsType<FValidConstFProperty>())
 			{
-				VisitFProperty(StackElement.VisitorTarget.Get<FValidConstFProperty>(), 
-					StackElement.OutputSchema.ToSharedRef(), StackElement.MemberPath, StackElement.PropertyFilter);
+				if (!VisitFProperty(StackElement.VisitorTarget.Get<FValidConstFProperty>(), 
+					StackElement.OutputSchema.ToSharedRef(), StackElement.MemberPath, StackElement.PropertyFilter))
+				{
+					return nullptr;
+				}
 			}
 		}
 	
 		return RootStackElement.OutputSchema;
 	}
 	
-	void VisitFProperty(const FValidConstFProperty Property, const TSharedRef<FJsonObject>& OutputSchema, 
+	bool VisitFProperty(const FValidConstFProperty Property, const TSharedRef<FJsonObject>& OutputSchema, 
 		const FJsonSchemaMemberPath& MemberPath, const FJsonSchemaPropertyFilter& PropertyFilter)
 	{
 		EditorMetadata.CurrentPropertyMemberPath = MemberPath;
@@ -123,16 +174,24 @@ public:
 		{
 			BuildJsonForBoolProperty(OutputSchema);
 		}
-		else if (CastField<FStrProperty>(Property) || CastField<FTextProperty>(Property) || CastField<FNameProperty>(Property))
+		else if (CastField<FStrProperty>(Property) || 
+			CastField<FTextProperty>(Property) || 
+			CastField<FNameProperty>(Property) ||
+			CastField<FUtf8StrProperty>(Property))
 		{
 			BuildJsonForStringProperty(OutputSchema);
 		}
 		else
 		{
-			// Fall back to string for unhandled property types.
-			BuildJsonForStringProperty(OutputSchema);
+			UE_LOG(LogJson, Error, TEXT(
+				"Property \"%s\" type %s unhandled during Json schema generation."),
+				*Property->GetName(), 
+				*Property->GetCPPType());
+			
+			return false;
 		}
 
+
 		// Add all property metadata, if any.
 		if (const TSharedPtr<FJsonObject> PropertyMetadataJson = EditorMetadata.GetPropertyMetadataForCurrentPropertyMemberPath())
 		{
@@ -142,16 +201,24 @@ public:
 				OutputSchema->SetField(*JsonMemberName, Tuple.Value);
 			}
 		}
+		
+		return true;
 	}
 	
-	void VisitUStruct(const FValidConstUStruct Struct, const TSharedRef<FJsonObject>& OutputSchema, 
+	bool VisitUStruct(const FValidConstUStruct Struct, const TSharedRef<FJsonObject>& OutputSchema, 
 		const FJsonSchemaMemberPath& MemberPath, const FJsonSchemaPropertyFilter& PropertyFilter,
 		const EVisitorStackElementFlags Flags)
 	{
 		EditorMetadata.CurrentPropertyMemberPath = MemberPath;
 
 		const bool bUseJsonRpcFormat = Struct->IsA(UFunction::StaticClass());
-
+		
+		// If using JSON RPC format, check that all properties can be accepted.
+		if (bUseJsonRpcFormat && !CanStructBeDirectlyRepresentedInJson(Struct, EditorMetadata))
+		{
+			return false;
+		}
+		
 		// Only include "type": "object" in this top-level schema if it's not a function, and so not following JSON RPC format.
 		if (!bUseJsonRpcFormat)
 		{
@@ -161,7 +228,7 @@ public:
 		// Special case handling for FJsonObjectWrapper, where we know this must be a Json object.   
 		if (Struct == FJsonObjectWrapper::StaticStruct())
 		{
-			return;
+			return true;
 		}
 
 		// Add all struct metadata.
@@ -260,11 +327,11 @@ public:
 				FunctionInputSchema->SetStringField(TEXT("type"), TEXT("object"));
 				OutputSchema->SetObjectField(TEXT("inputSchema"), FunctionInputSchema);
 				
-				if (FunctionInputPropertiesSchema->Values.Num() > 0)
+				if (!FunctionInputPropertiesSchema->Values.IsEmpty())
 				{
 					FunctionInputSchema->SetObjectField(TEXT("properties"), FunctionInputPropertiesSchema);
 
-					if (FunctionInputRequiredPropertiesArray.Num() > 0)
+					if (!FunctionInputRequiredPropertiesArray.IsEmpty())
 					{
 						FunctionInputSchema->SetArrayField(TEXT("required"), FunctionInputRequiredPropertiesArray);
 					}
@@ -272,7 +339,7 @@ public:
 			}
 	
 			// Output schema is optional.
-			if (FunctionOutputPropertiesSchema->Values.Num() > 0)
+			if (!FunctionOutputPropertiesSchema->Values.IsEmpty())
 			{
 				const TSharedPtr<FJsonObject> FunctionOutputSchema = MakeShared<FJsonObject>();
 				FunctionOutputSchema->SetStringField(TEXT("type"), TEXT("object"));
@@ -282,7 +349,7 @@ public:
 				{
 					FunctionOutputSchema->SetObjectField(TEXT("properties"), FunctionOutputPropertiesSchema);
 					
-					if (FunctionOutputRequiredPropertiesArray.Num() > 0)
+					if (!FunctionOutputRequiredPropertiesArray.IsEmpty())
 					{
 						FunctionOutputSchema->SetArrayField(TEXT("required"), FunctionOutputRequiredPropertiesArray);
 					}
@@ -300,15 +367,17 @@ public:
 				MaybeIncludeProperty(Property, PropertiesSchema, RequiredPropertiesArray);
 			}
 	
-			if (PropertiesSchema->Values.Num() > 0)
+			if (!PropertiesSchema->Values.IsEmpty())
 			{
 				OutputSchema->SetObjectField(TEXT("properties"), PropertiesSchema);
-				if (RequiredPropertiesArray.Num() > 0)
+				if (!RequiredPropertiesArray.IsEmpty())
 				{
 					OutputSchema->SetArrayField(TEXT("required"), RequiredPropertiesArray);
 				}
 			}
 		}
+		
+		return true;
 	}
 	
 private:
diff --git a/Engine/Source/Runtime/JsonUtilities/Private/JsonSchema/JsonSchemaVisitorStack.h b/Engine/Source/Runtime/JsonUtilities/Private/JsonSchema/JsonSchemaVisitorStack.h
--- a/Engine/Source/Runtime/JsonUtilities/Private/JsonSchema/JsonSchemaVisitorStack.h
+++ b/Engine/Source/Runtime/JsonUtilities/Private/JsonSchema/JsonSchemaVisitorStack.h
@@ -19,8 +19,7 @@ namespace UE::JsonSchema
 	enum class EVisitorStackElementFlags : uint8
 	{
 		None = 0x0,
-		IsRootStruct = 0x1,
-		IsFunctionWorldContextProperty = 0x2
+		IsRootStruct = 0x1
 	};
 	ENUM_CLASS_FLAGS(EVisitorStackElementFlags);
 
@@ -37,11 +36,7 @@ namespace UE::JsonSchema
 			PropertyFilter(InPropertyFilter),
 			OutputSchema(InOutputSchema)
 		{
-			if (VisitorTarget.IsType<FValidConstUStruct>())
-			{
-				check(!EnumHasAnyFlags(Flags, EVisitorStackElementFlags::IsFunctionWorldContextProperty));
-			}
-			else if (VisitorTarget.IsType<FValidConstFProperty>())
+			if (VisitorTarget.IsType<FValidConstFProperty>())
 			{
 				check(!EnumHasAnyFlags(Flags, EVisitorStackElementFlags::IsRootStruct));
 			}
diff --git a/Engine/Source/Runtime/JsonUtilities/Private/JsonUtilitiesLog.cpp b/Engine/Source/Runtime/JsonUtilities/Private/JsonUtilitiesLog.cpp
deleted file mode 100644
diff --git a/Engine/Source/Runtime/JsonUtilities/Private/JsonUtilitiesLog.h b/Engine/Source/Runtime/JsonUtilities/Private/JsonUtilitiesLog.h
deleted file mode 100644
